Removed UE4 occurrences from AIModule's code

diff --git a/Engine/Plugins/AI/EnvironmentQueryEditor/Source/EnvironmentQueryEditor/Private/EnvironmentQueryEditor.cpp b/Engine/Plugins/AI/EnvironmentQueryEditor/Source/EnvironmentQueryEditor/Private/EnvironmentQueryEditor.cpp
--- a/Engine/Plugins/AI/EnvironmentQueryEditor/Source/EnvironmentQueryEditor/Private/EnvironmentQueryEditor.cpp
+++ b/Engine/Plugins/AI/EnvironmentQueryEditor/Source/EnvironmentQueryEditor/Private/EnvironmentQueryEditor.cpp
@@ -32,8 +32,8 @@
 namespace FEnvironmentQueryHelper
 {
 	static const FString StatFileDescription = LOCTEXT("FileTypeDescription", "EQS Stat File").ToString();
-	static const FString LoadFileTypes = FString::Printf(TEXT("%s (*.ue4eqs)|*.ue4eqs"), *StatFileDescription);
-	static const FString SaveFileTypes = FString::Printf(TEXT("%s (*.ue4eqs)|*.ue4eqs"), *StatFileDescription);
+	static const FString LoadFileTypes = FString::Printf(TEXT("%s (*.ue_eqs)|*.ue_eqs"), *StatFileDescription);
+	static const FString SaveFileTypes = FString::Printf(TEXT("%s (*.ue_eqs)|*.ue_eqs"), *StatFileDescription);
 }
 
 const FName FEnvironmentQueryEditor::EQSUpdateGraphTabId(TEXT("EnvironmentQueryEditor_UpdateGraph"));
diff --git a/Engine/Source/Runtime/AIModule/Private/BehaviorTree/BlackboardData.cpp b/Engine/Source/Runtime/AIModule/Private/BehaviorTree/BlackboardData.cpp
--- a/Engine/Source/Runtime/AIModule/Private/BehaviorTree/BlackboardData.cpp
+++ b/Engine/Source/Runtime/AIModule/Private/BehaviorTree/BlackboardData.cpp
@@ -155,7 +155,7 @@ void UBlackboardData::PostLoad()
 	Super::PostLoad();
 
 	// we cache some information based on Parent asset
-	// but while UE4 guarantees the Parent is already loaded
+	// but while UnrealEngine guarantees the Parent is already loaded
 	// it does not guarantee that it's PostLoad has been called
 	// Following is a little hack that's widely used in the
 	// engine to address this
diff --git a/Engine/Source/Runtime/AIModule/Private/EnvironmentQuery/EnvQueryManager.cpp b/Engine/Source/Runtime/AIModule/Private/EnvironmentQuery/EnvQueryManager.cpp
--- a/Engine/Source/Runtime/AIModule/Private/EnvironmentQuery/EnvQueryManager.cpp
+++ b/Engine/Source/Runtime/AIModule/Private/EnvironmentQuery/EnvQueryManager.cpp
@@ -1000,7 +1000,7 @@ bool UEnvQueryManager::Exec(UWorld* Inworld, const TCHAR* Cmd, FOutputDevice& Ar
 #if USE_EQS_DEBUGGER
 	if (FParse::Command(&Cmd, TEXT("DumpEnvQueryStats")))
 	{
-		const FString FileName = FPaths::CreateTempFilename(*FPaths::ProjectLogDir(), TEXT("EnvQueryStats"), TEXT(".ue4eqs"));
+		const FString FileName = FPaths::CreateTempFilename(*FPaths::ProjectLogDir(), TEXT("EnvQueryStats"), TEXT(".ue_eqs"));
 
 		FEQSDebugger::SaveStats(FileName);
 		return true;
