Adding drag and drop functionality to EQS queries
- EQS queries can now be dragged from the content browser and dropped into the editor world and the EQSTestingPawn will be autocreated and setup with the selected EQS query

diff --git a/Engine/Plugins/AI/EnvironmentQueryEditor/Source/EnvironmentQueryEditor/Private/ActorFactoryEnvironmentQuery.cpp b/Engine/Plugins/AI/EnvironmentQueryEditor/Source/EnvironmentQueryEditor/Private/ActorFactoryEnvironmentQuery.cpp
new file mode 100644
--- /dev/null
+++ b/Engine/Plugins/AI/EnvironmentQueryEditor/Source/EnvironmentQueryEditor/Private/ActorFactoryEnvironmentQuery.cpp
@@ -0,0 +1,44 @@
+// Copyright Epic Games, Inc. All Rights Reserved.
+
+#include "ActorFactoryEnvironmentQuery.h"
+
+#include "EnvironmentQuery/EnvQuery.h"
+#include "EnvironmentQuery/EQSTestingPawn.h"
+
+#define LOCTEXT_NAMESPACE "ActorFactoryEnvironmentQuery"
+
+UActorFactoryEnvironmentQuery::UActorFactoryEnvironmentQuery()
+{
+	DisplayName = LOCTEXT("EnvironmentQueryActorDisplayName", "Environment Query");
+	NewActorClass = AEQSTestingPawn::StaticClass();
+}
+
+void UActorFactoryEnvironmentQuery::PostSpawnActor(UObject* Asset, AActor* NewActor)
+{
+	Super::PostSpawnActor(Asset, NewActor);
+
+	UEnvQuery* EnvQuery = CastChecked<UEnvQuery>(Asset);
+	AEQSTestingPawn* NewEQSTestingPawn = CastChecked<AEQSTestingPawn>(NewActor);
+
+	NewEQSTestingPawn->QueryTemplate = EnvQuery;
+	NewEQSTestingPawn->QueryTemplate->CollectQueryParams(*NewEQSTestingPawn, NewEQSTestingPawn->QueryConfig);
+	NewEQSTestingPawn->RunEQSQuery();
+}
+
+bool UActorFactoryEnvironmentQuery::CanCreateActorFrom(const FAssetData& AssetData, FText& OutErrorMsg)
+{
+	if (!AssetData.IsValid() || !AssetData.GetClass()->IsChildOf(UEnvQuery::StaticClass()))
+	{
+		OutErrorMsg = LOCTEXT("NoEnvironmentQuery", "A valid environment query must be specified.");
+		return false;
+	}
+
+	return true;
+}
+
+UObject* UActorFactoryEnvironmentQuery::GetAssetFromActorInstance(AActor* ActorInstance)
+{
+	return CastChecked<AEQSTestingPawn>(ActorInstance)->QueryTemplate;
+}
+
+#undef LOCTEXT_NAMESPACE
\ No newline at end of file
diff --git a/Engine/Plugins/AI/EnvironmentQueryEditor/Source/EnvironmentQueryEditor/Public/ActorFactoryEnvironmentQuery.h b/Engine/Plugins/AI/EnvironmentQueryEditor/Source/EnvironmentQueryEditor/Public/ActorFactoryEnvironmentQuery.h
new file mode 100644
--- /dev/null
+++ b/Engine/Plugins/AI/EnvironmentQueryEditor/Source/EnvironmentQueryEditor/Public/ActorFactoryEnvironmentQuery.h
@@ -0,0 +1,28 @@
+// Copyright Epic Games, Inc. All Rights Reserved.
+
+#pragma once
+
+#include "UObject/ObjectMacros.h"
+#include "ActorFactories/ActorFactory.h"
+#include "ActorFactoryEnvironmentQuery.generated.h"
+
+class AActor;
+struct FAssetData;
+
+UCLASS(MinimalAPI, config=Editor, collapsecategories, hidecategories=Object)
+class UActorFactoryEnvironmentQuery : public UActorFactory
+{
+	GENERATED_BODY()
+
+	UActorFactoryEnvironmentQuery();
+
+protected:
+	//~ Begin UActorFactory Interface
+	virtual void PostSpawnActor(UObject* Asset, AActor* NewActor) override;
+	virtual bool CanCreateActorFrom(const FAssetData& AssetData, FText& OutErrorMsg) override;
+	virtual UObject* GetAssetFromActorInstance(AActor* ActorInstance) override;
+	//~ End UActorFactory Interface
+};
+
+
+
