Added a comment to explain why Super::CanCreateActorFrom(...) isn't used

#rnx

diff --git a/Engine/Plugins/AI/EnvironmentQueryEditor/Source/EnvironmentQueryEditor/Private/ActorFactoryEnvironmentQuery.cpp b/Engine/Plugins/AI/EnvironmentQueryEditor/Source/EnvironmentQueryEditor/Private/ActorFactoryEnvironmentQuery.cpp
--- a/Engine/Plugins/AI/EnvironmentQueryEditor/Source/EnvironmentQueryEditor/Private/ActorFactoryEnvironmentQuery.cpp
+++ b/Engine/Plugins/AI/EnvironmentQueryEditor/Source/EnvironmentQueryEditor/Private/ActorFactoryEnvironmentQuery.cpp
@@ -33,6 +33,8 @@ bool UActorFactoryEnvironmentQuery::CanCreateActorFrom(const FAssetData& AssetDa
 		return false;
 	}
 
+	// Not calling Super since it's checking if AssetData matches the actor class we want to spawn, and that's not the case here.
+	// We're spawning a EQSTestingPawn to put a given EQS query in the world context.
 	return true;
 }
 
