Fix editor crash when EQS details tab is closed and selecting a test or generator in an EQS query.
PR #9473: Eqs crash (Contributed by [REDACTED_PERSON])

diff --git a/Engine/Plugins/AI/EnvironmentQueryEditor/Source/EnvironmentQueryEditor/Private/EnvironmentQueryEditor.cpp b/Engine/Plugins/AI/EnvironmentQueryEditor/Source/EnvironmentQueryEditor/Private/EnvironmentQueryEditor.cpp
--- a/Engine/Plugins/AI/EnvironmentQueryEditor/Source/EnvironmentQueryEditor/Private/EnvironmentQueryEditor.cpp
+++ b/Engine/Plugins/AI/EnvironmentQueryEditor/Source/EnvironmentQueryEditor/Private/EnvironmentQueryEditor.cpp
@@ -299,13 +299,22 @@ void FEnvironmentQueryEditor::OnSelectedNodesChanged(const TSet<class UObject*>&
 		}
 	}
 
-	if (Selection.Num() == 1)
+	TSharedPtr<FTabManager> TabManagerPtr = GetTabManager();
+	if (TabManagerPtr)
 	{
-		DetailsView->SetObjects(Selection);
+		TabManagerPtr->TryInvokeTab(EQSPropertiesTabId);
 	}
-	else
+	
+	if (DetailsView)
 	{
-		DetailsView->SetObject(NULL);
+		if (Selection.Num() == 1)
+		{
+			DetailsView->SetObjects(Selection);
+		}
+		else
+		{
+			DetailsView->SetObject(NULL);
+		}
 	}
 }
 
