PR #11958: Fix EQSEditorModule custom Toolbar Extenders not being applied

diff --git a/Engine/Plugins/AI/EnvironmentQueryEditor/Source/EnvironmentQueryEditor/Private/EnvironmentQueryEditor.cpp b/Engine/Plugins/AI/EnvironmentQueryEditor/Source/EnvironmentQueryEditor/Private/EnvironmentQueryEditor.cpp
--- a/Engine/Plugins/AI/EnvironmentQueryEditor/Source/EnvironmentQueryEditor/Private/EnvironmentQueryEditor.cpp
+++ b/Engine/Plugins/AI/EnvironmentQueryEditor/Source/EnvironmentQueryEditor/Private/EnvironmentQueryEditor.cpp
@@ -122,6 +122,7 @@ void FEnvironmentQueryEditor::InitEnvironmentQueryEditor( const EToolkitMode::Ty
 	
 	FEnvironmentQueryEditorModule& EnvironmentQueryEditorModule = FModuleManager::LoadModuleChecked<FEnvironmentQueryEditorModule>( "EnvironmentQueryEditor" );
 	AddMenuExtender(EnvironmentQueryEditorModule.GetMenuExtensibilityManager()->GetAllExtenders(GetToolkitCommands(), GetEditingObjects()));
+	AddToolbarExtender(EnvironmentQueryEditorModule.GetToolBarExtensibilityManager()->GetAllExtenders(GetToolkitCommands(), GetEditingObjects()));
 
 	BindCommands();
 	ExtendToolbar();
