[EQS] Marked EQS editor as production ready

diff --git a/Engine/Plugins/AI/EnvironmentQueryEditor/EnvironmentQueryEditor.uplugin b/Engine/Plugins/AI/EnvironmentQueryEditor/EnvironmentQueryEditor.uplugin
--- a/Engine/Plugins/AI/EnvironmentQueryEditor/EnvironmentQueryEditor.uplugin
+++ b/Engine/Plugins/AI/EnvironmentQueryEditor/EnvironmentQueryEditor.uplugin
@@ -1,8 +1,8 @@
 {
 	"FileVersion" : 3,
 	"Version" : 1,
-	"VersionName" : "0.1",
-	"FriendlyName" : "Environment Query Editor (Experimental)",
+	"VersionName" : "1.0",
+	"FriendlyName" : "Environment Query Editor",
 	"Description" : "Allows editing of Environment Query assets, which are used by the AI to collect data about the environment/world",
 	"Category" : "AI",
 	"CreatedBy" : "Epic Games, Inc.",
@@ -12,7 +12,6 @@
 	"SupportURL" : "",
 	"EnabledByDefault" : true,
 	"CanContainContent" : false,
-	"IsBetaVersion" : true,
 	"Installed" : false,
 	"Modules" : [
 		{
