Allow EnvQuery TraceByProfile to set the shape of the trace used.

diff --git a/Engine/Plugins/AI/EnvironmentQueryEditor/Source/EnvironmentQueryEditor/Private/DetailCustomizations/EnvTraceDataCustomization.cpp b/Engine/Plugins/AI/EnvironmentQueryEditor/Source/EnvironmentQueryEditor/Private/DetailCustomizations/EnvTraceDataCustomization.cpp
--- a/Engine/Plugins/AI/EnvironmentQueryEditor/Source/EnvironmentQueryEditor/Private/DetailCustomizations/EnvTraceDataCustomization.cpp
+++ b/Engine/Plugins/AI/EnvironmentQueryEditor/Source/EnvironmentQueryEditor/Private/DetailCustomizations/EnvTraceDataCustomization.cpp
@@ -73,7 +73,7 @@ void FEnvTraceDataCustomization::CustomizeChildren( TSharedRef<class IPropertyHa
 		.Visibility(TAttribute<EVisibility>::Create(TAttribute<EVisibility>::FGetter::CreateSP(this, &FEnvTraceDataCustomization::GetGeometryByChannelVisibility)));
 
 	StructBuilder.AddProperty(PropTraceShape.ToSharedRef())
-		.Visibility(TAttribute<EVisibility>::Create(TAttribute<EVisibility>::FGetter::CreateSP(this, &FEnvTraceDataCustomization::GetGeometryByChannelVisibility)));
+		.Visibility(TAttribute<EVisibility>::Create(TAttribute<EVisibility>::FGetter::CreateSP(this, &FEnvTraceDataCustomization::GetGeometryVisibility)));
 
 	StructBuilder.AddProperty(PropTraceProfileName.ToSharedRef())
 		.Visibility(TAttribute<EVisibility>::Create(TAttribute<EVisibility>::FGetter::CreateSP(this, &FEnvTraceDataCustomization::GetGeometryByProfileVisibility)));
